LiveLink Control Rig

-Added node to access basic property values by LiveLink subject and property name from any LiveLink Role

diff --git a/Engine/Plugins/Experimental/LiveLinkControlRig/Source/LiveLinkControlRig/Private/LiveLinkRigUnits.cpp b/Engine/Plugins/Experimental/LiveLinkControlRig/Source/LiveLinkControlRig/Private/LiveLinkRigUnits.cpp
--- a/Engine/Plugins/Experimental/LiveLinkControlRig/Source/LiveLinkControlRig/Private/LiveLinkRigUnits.cpp
+++ b/Engine/Plugins/Experimental/LiveLinkControlRig/Source/LiveLinkControlRig/Private/LiveLinkRigUnits.cpp
@@ -7,6 +7,7 @@
 #include "Roles/LiveLinkAnimationRole.h"
 #include "Roles/LiveLinkTransformRole.h"
 #include "Roles/LiveLinkTransformTypes.h"
+#include "Roles/LiveLinkBasicRole.h"
 
 ILiveLinkClient* LiveLinkControlRigUtilities::TryGetLiveLinkClient()
 {
@@ -136,4 +137,28 @@ FRigUnit_LiveLinkEvaluteFrameTransform_Execute()
 		Context.DrawInterface->DrawHierarchy(DebugDrawOffset, DebugHierarchy, EControlRigDrawHierarchyMode::Axes, 1.0f, DebugColor, 1.0f);
 		*/
 	}
+}
+
+FRigUnit_LiveLinkEvaluateBasicValue_Execute()
+{
+	// Get value by property name from basic data struct
+	if (ILiveLinkClient* LiveLinkClient = LiveLinkControlRigUtilities::TryGetLiveLinkClient())
+	{
+		FLiveLinkSubjectFrameData FrameData;
+		if (LiveLinkClient->EvaluateFrame_AnyThread(SubjectName, ULiveLinkBasicRole::StaticClass(), FrameData))
+		{
+			const FLiveLinkBaseStaticData* StaticData = FrameData.StaticData.Cast<FLiveLinkBaseStaticData>();
+			if(StaticData)
+			{
+				const int32 FoundIndex = StaticData->PropertyNames.Find(PropertyName);
+
+				if(FoundIndex != INDEX_NONE)
+				{
+					FLiveLinkBaseFrameData* BasicData = FrameData.FrameData.Cast<FLiveLinkBaseFrameData>();
+					Value = BasicData->PropertyValues[FoundIndex];
+				}
+			}
+
+		}
+	}
 }
\ No newline at end of file
diff --git a/Engine/Plugins/Experimental/LiveLinkControlRig/Source/LiveLinkControlRig/Public/LiveLinkRigUnits.h b/Engine/Plugins/Experimental/LiveLinkControlRig/Source/LiveLinkControlRig/Public/LiveLinkRigUnits.h
--- a/Engine/Plugins/Experimental/LiveLinkControlRig/Source/LiveLinkControlRig/Public/LiveLinkRigUnits.h
+++ b/Engine/Plugins/Experimental/LiveLinkControlRig/Source/LiveLinkControlRig/Public/LiveLinkRigUnits.h
@@ -119,4 +119,25 @@ struct LIVELINKCONTROLRIG_API FRigUnit_LiveLinkEvaluteFrameTransform : public FR
 
 	UPROPERTY(meta = (Output))
 	FTransform Transform;
+};
+/**
+ * Evaluate current Live Link Basic float property data associated with supplied subject
+ */
+USTRUCT(meta = (DisplayName = "Get Basic LiveLink Data", Category = "Live Link"))
+struct LIVELINKCONTROLRIG_API FRigUnit_LiveLinkEvaluateBasicValue : public FRigUnit_LiveLinkBase
+{
+	GENERATED_BODY()
+
+	/** Execute logic for this rig unit */
+	RIGVM_METHOD()
+	virtual void Execute() override;
+
+	UPROPERTY(meta = (Input))
+	FName SubjectName;
+
+	UPROPERTY(meta = (Input))
+	FName PropertyName;
+
+	UPROPERTY(meta = (Output))
+	float Value = 0.f;
 };
\ No newline at end of file
