LiveLinkControlRig:

-Fix for crash caused by a case where a BP user can create a virtual subject with a different number of PropertyNames and PropertyValues

diff --git a/Engine/Plugins/Experimental/LiveLinkControlRig/Source/LiveLinkControlRig/Private/LiveLinkRigUnits.cpp b/Engine/Plugins/Experimental/LiveLinkControlRig/Source/LiveLinkControlRig/Private/LiveLinkRigUnits.cpp
--- a/Engine/Plugins/Experimental/LiveLinkControlRig/Source/LiveLinkControlRig/Private/LiveLinkRigUnits.cpp
+++ b/Engine/Plugins/Experimental/LiveLinkControlRig/Source/LiveLinkControlRig/Private/LiveLinkRigUnits.cpp
@@ -155,10 +155,12 @@ FRigUnit_LiveLinkEvaluateBasicValue_Execute()
 				if(FoundIndex != INDEX_NONE)
 				{
 					FLiveLinkBaseFrameData* BasicData = FrameData.FrameData.Cast<FLiveLinkBaseFrameData>();
+					if(StaticData->PropertyNames.Num() == BasicData->PropertyValues.Num()) //Catch case where the number of PropertyNames and PropertyValues don't match. This can happen in Blueprint Virtual Subjects.
+					{
 					Value = BasicData->PropertyValues[FoundIndex];
+					}
 				}
 			}
-
 		}
 	}
 }
\ No newline at end of file
